Sequencer Control Rig: Make control rig work through the anim mixer and be able to be mixed in any order with animations.
For now, this doesn't work underneath the 'anim mixer' track- this will be the next step. But this gets it working through the anim mixer code without creating a new custom anim instance. This effectively enables Control Rig in Sequencer to work with UAF.
A few notes:
* This follows the same pattern of creating a new 'mixed control rig system' that tags existing control rig entities (just the base rig entity) so that the original control rig system doesn't process certain things, in particular, that it doesn't create a new custom anim instance or do anything having to do with cached restore state for that anim instance. However unlike the skeletal animation system refactor, this does leave the initial control rig system to do most things, since most operations there happen on the ControlRig instance itself, which is the same in both cases.
* This is implemented in a new 'MovieSceneMixedControlRig' experimental plugin, due to the dependencies on the Control Rig plugin, the experimental UAF plugin, the experimental Anim Mixer plugin, and the experimental UAF Control Rig plugin. To test this, all of those plugins must be enabled in your project. It's been tested in AnimSandbox and the ContentExamples project.
* By default the control rig is still going through the old evaluation template workflow rather than its new ECS workflow. This change requires the change to the ECS workflow via setting ControlRig.UseLegacySequencerTemplate to false.
* There are still some open questions about how layered control rigs vs. regular control rigs should work with this. Currently both are allowed to be weighted and there is always a mix step therefore added, meaning from a mixer perspective there is no real difference between them, but perhaps we want there to be, e.g. perhaps Control Rigs should always overwrite the previous pose entirely.
* It was challenging to make this 100% backwards compatible with sequences having an animation and multiple control rig layers due to ordering issues, but some magic numbers have been used in the default priorities to apply for animation vs. control rig priorities that don't have anim mixers. Long-term the true backwards compatibility path will likely be to add these automatically to an anim mixer track in postload which should help with some of that since it'll just be about order rather than magic numbers.

diff --git a/Engine/Plugins/Animation/ControlRig/Source/ControlRig/Private/Sequencer/MovieSceneControlRigSystem.cpp b/Engine/Plugins/Animation/ControlRig/Source/ControlRig/Private/Sequencer/MovieSceneControlRigSystem.cpp
--- a/Engine/Plugins/Animation/ControlRig/Source/ControlRig/Private/Sequencer/MovieSceneControlRigSystem.cpp
+++ b/Engine/Plugins/Animation/ControlRig/Source/ControlRig/Private/Sequencer/MovieSceneControlRigSystem.cpp
@@ -1666,9 +1666,10 @@ void UMovieSceneControlRigParameterEvaluatorSystem::OnInstantiation()
 		.Write(ControlRigComponents->BaseControlRigEvalData)
 		.PassthroughFilter(FEntityComponentFilter().All({ BuiltInComponents->Tags.RestoreState }))
 		.PassthroughFilter(FEntityComponentFilter().All({ BuiltInComponents->WeightAndEasingResult }))
+		.PassthroughFilter(FEntityComponentFilter().All({ TracksComponents->Tags.AnimMixerPoseProducer}))
 		.FilterNone({ BuiltInComponents->Tags.NeedsUnlink })
 		.Iterate_PerEntity(&Linker->EntityManager, [this, &EncounteredRigs](FMovieSceneEntityID EntityID, const FGuid& ObjectBindingID, FRootInstanceHandle RootInstanceHandle, const FMovieSceneSequenceID* OptSequenceID,
-				const FControlRigSourceData& ControlRigSource, FBaseControlRigEvalData& OutBaseData, bool bWantsRestoreState, bool bHasWeight)
+				const FControlRigSourceData& ControlRigSource, FBaseControlRigEvalData& OutBaseData, bool bWantsRestoreState, bool bHasWeight, bool bAnimMixerPoseProducer)
 			{
 
 				OutBaseData.bIsActive = false;
@@ -1689,6 +1690,7 @@ void UMovieSceneControlRigParameterEvaluatorSystem::OnInstantiation()
 							RootInstanceHandle,
 							bWantsRestoreState,
 							bHasWeight,
+							bAnimMixerPoseProducer,
 							ControlRigSource,
 							OutBaseData,
 							EncounteredRigs);
@@ -1876,6 +1878,7 @@ void UMovieSceneControlRigParameterEvaluatorSystem::InitializeBaseRigComponent(
 	UE::MovieScene::FRootInstanceHandle RootInstanceHandle,
 	bool bWantsRestoreState,
 	bool bHasWeight,
+	bool bAnimMixerPoseProducer,
 	UE::MovieScene::FControlRigSourceData ControlRigSource,
 	UE::MovieScene::FBaseControlRigEvalData& OutBaseData,
 	TMap<UControlRig*, UE::MovieScene::FBaseControlRigEvalData*>& OutBaseComponentTracker)
@@ -1908,12 +1911,16 @@ void UMovieSceneControlRigParameterEvaluatorSystem::InitializeBaseRigComponent(
 		return;
 	}
 
-	// Begin tracking this entity for the base pre-animated state that keeps the rig bound to our target object
-	this->ControlRigStorage->BeginTrackingEntity(
-		EntityID,
-		bWantsRestoreState,
-		RootInstanceHandle,
-		ControlRig);
+	// The anim mixer handles restore state for binding to the skeletal mesh component
+	if (!bAnimMixerPoseProducer)
+	{
+		// Begin tracking this entity for the base pre-animated state that keeps the rig bound to our target object
+		this->ControlRigStorage->BeginTrackingEntity(
+			EntityID,
+			bWantsRestoreState,
+			RootInstanceHandle,
+			ControlRig);
+	}
 
 	if (!ControlRig->GetObjectBinding())
 	{
@@ -1959,8 +1966,12 @@ void UMovieSceneControlRigParameterEvaluatorSystem::InitializeBaseRigComponent(
 
 	OutBaseData.WeakControlRig = ControlRig;
 
-	// Cache pre-animated value now that the control rig is bound to the correct object
-	this->ControlRigStorage->CachePreAnimatedValue(ControlRig);
+	// The anim mixer handles restore state for binding to the skeletal mesh component
+	if (!bAnimMixerPoseProducer)
+	{
+		// Cache pre-animated value now that the control rig is bound to the correct object
+		this->ControlRigStorage->CachePreAnimatedValue(ControlRig);
+	}
 
 	FBaseControlRigEvalData** ExistingData = OutBaseComponentTracker.Find(ControlRig);
 	if (ExistingData)
@@ -1986,24 +1997,34 @@ void UMovieSceneControlRigParameterEvaluatorSystem::InitializeBaseRigComponent(
 		OutBaseData.bIsActive = true;
 		OutBaseComponentTracker.Add(ControlRig, &OutBaseData);
 	}
+	
 
-#if WITH_EDITOR
+#if WITH_EDITOR	
+	// The anim mixer handles binding to the skeletal mesh component
 	TWeakObjectPtr<UAnimInstance> PreviousAnimInstanceWeakPtr = nullptr;
-			
-	if (ControlRig->GetObjectBinding())
+	if (!bAnimMixerPoseProducer)
 	{
-		if (const USkeletalMeshComponent* SkeletalMeshComponent = Cast<USkeletalMeshComponent>(ControlRig->GetObjectBinding()->GetBoundObject()))
+		if (ControlRig->GetObjectBinding())
 		{
-			PreviousAnimInstanceWeakPtr = SkeletalMeshComponent->GetAnimInstance();
+			if (const USkeletalMeshComponent* SkeletalMeshComponent = Cast<USkeletalMeshComponent>(ControlRig->GetObjectBinding()->GetBoundObject()))
+			{
+				PreviousAnimInstanceWeakPtr = SkeletalMeshComponent->GetAnimInstance();
+			}
 		}
 	}
 #endif
 
-	bool bWasCreated = FControlRigBindingHelper::BindToSequencerInstance(ControlRig);
+	// We bind to sequencer instance if we're not an anim mixer pose producer or if our bound object is not a skm component.
+	bool bWasCreated = false;
+	if (!bAnimMixerPoseProducer || (ControlRig->GetObjectBinding() && Cast<USkeletalMeshComponent>(ControlRig->GetObjectBinding()->GetBoundObject()) == nullptr))
+	{
+		bWasCreated = FControlRigBindingHelper::BindToSequencerInstance(ControlRig);
+	}
+	
 	this->AccumulatedValues.InitializeRig(Track, ControlRig);
 
 #if WITH_EDITOR
-	if (GEditor && bWasCreated)
+	if (!bAnimMixerPoseProducer && GEditor && bWasCreated)
 	{
 		if (ControlRig->GetObjectBinding())
 		{
@@ -2020,12 +2041,12 @@ void UMovieSceneControlRigParameterEvaluatorSystem::InitializeBaseRigComponent(
 					}
 				});
 								
+
 				PreCompileHandles.Add(PreCompileHandle);
 			}
 		}
 	}
 #endif
-	
 }
 
 #if WITH_EDITOR
diff --git a/Engine/Plugins/Animation/ControlRig/Source/ControlRig/Public/Sequencer/MovieSceneControlRigComponentTypes.h b/Engine/Plugins/Animation/ControlRig/Source/ControlRig/Public/Sequencer/MovieSceneControlRigComponentTypes.h
--- a/Engine/Plugins/Animation/ControlRig/Source/ControlRig/Public/Sequencer/MovieSceneControlRigComponentTypes.h
+++ b/Engine/Plugins/Animation/ControlRig/Source/ControlRig/Public/Sequencer/MovieSceneControlRigComponentTypes.h
@@ -44,8 +44,8 @@ struct FBaseControlRigEvalData
 struct FControlRigComponentTypes
 {
 public:
-	static FControlRigComponentTypes* Get();
-	static void Destroy();
+	CONTROLRIG_API static FControlRigComponentTypes* Get();
+	CONTROLRIG_API static void Destroy();
 
 	TComponentTypeID<FControlRigSourceData> ControlRigSource;
 
@@ -66,7 +66,7 @@ public:
 	} Tags;
 
 private:
-	FControlRigComponentTypes();
+	CONTROLRIG_API FControlRigComponentTypes();
 };
 
 
diff --git a/Engine/Plugins/Animation/ControlRig/Source/ControlRig/Public/Sequencer/MovieSceneControlRigParameterTrack.h b/Engine/Plugins/Animation/ControlRig/Source/ControlRig/Public/Sequencer/MovieSceneControlRigParameterTrack.h
--- a/Engine/Plugins/Animation/ControlRig/Source/ControlRig/Public/Sequencer/MovieSceneControlRigParameterTrack.h
+++ b/Engine/Plugins/Animation/ControlRig/Source/ControlRig/Public/Sequencer/MovieSceneControlRigParameterTrack.h
@@ -261,7 +261,7 @@ public:
 	static CONTROLRIG_API FColor LayeredRigTrackColor;
 
 public:
-	UControlRig* GetGameWorldControlRig(UWorld* InWorld);
+	CONTROLRIG_API UControlRig* GetGameWorldControlRig(UWorld* InWorld);
 	bool IsAGameInstance(const UControlRig* InControlRig, const bool bCheckValidWorld = false) const;
     
 private:
diff --git a/Engine/Plugins/Animation/ControlRig/Source/ControlRig/Public/Sequencer/MovieSceneControlRigSystem.h b/Engine/Plugins/Animation/ControlRig/Source/ControlRig/Public/Sequencer/MovieSceneControlRigSystem.h
--- a/Engine/Plugins/Animation/ControlRig/Source/ControlRig/Public/Sequencer/MovieSceneControlRigSystem.h
+++ b/Engine/Plugins/Animation/ControlRig/Source/ControlRig/Public/Sequencer/MovieSceneControlRigSystem.h
@@ -161,7 +161,7 @@ namespace UE::MovieScene
  * 
  * Instantiation:
  *   - Initializes rigs by ensuring the anim instance set up and bound to the correct rig
- *   - Tracks contriguting control rig parameters by rg
+ *   - Tracks contributing control rig parameters by rig
  *   - Caches pre-animated state
  *   - Assigns accumulation indices for each parameter
  *   - Caches initial values if necessary
@@ -206,6 +206,7 @@ private:
 		UE::MovieScene::FRootInstanceHandle RootInstanceHandle,
 		bool bWantsRestoreState,
 		bool bHasWeight,
+		bool bAnimMixerPoseProducer,
 		UE::MovieScene::FControlRigSourceData ControlRigSource,
 		UE::MovieScene::FBaseControlRigEvalData& OutBaseData,
 		TMap<UControlRig*, UE::MovieScene::FBaseControlRigEvalData*>& OutBaseComponentTracker
diff --git a/Engine/Plugins/Experimental/MovieSceneAnimMixer/Source/MovieSceneAnimMixer/Private/Systems/MovieSceneMixedSkeletalAnimationSystem.cpp b/Engine/Plugins/Experimental/MovieSceneAnimMixer/Source/MovieSceneAnimMixer/Private/Systems/MovieSceneMixedSkeletalAnimationSystem.cpp
--- a/Engine/Plugins/Experimental/MovieSceneAnimMixer/Source/MovieSceneAnimMixer/Private/Systems/MovieSceneMixedSkeletalAnimationSystem.cpp
+++ b/Engine/Plugins/Experimental/MovieSceneAnimMixer/Source/MovieSceneAnimMixer/Private/Systems/MovieSceneMixedSkeletalAnimationSystem.cpp
@@ -136,7 +136,7 @@ struct FSkeletalAnimMixerMutation : IMovieSceneEntityMutation
 				{
 					// Set a default target- Automatic will resolve to the appropriate place, or the Anim Mixer track will override it.
 					OutAnimTargets[Index] = TInstancedStruct<FMovieSceneMixedAnimationTarget>::Make();
-					OutPriorities[Index] = 0;
+					OutPriorities[Index] = -100000;
 				}
 
 				UMovieSceneCommonAnimationTrack* Track = SkeletalAnim.Section->GetTypedOuter<UMovieSceneCommonAnimationTrack>();
diff --git a/Engine/Plugins/Experimental/MovieSceneMixedControlRig/MovieSceneMixedControlRig.uplugin b/Engine/Plugins/Experimental/MovieSceneMixedControlRig/MovieSceneMixedControlRig.uplugin
new file mode 100644
--- /dev/null
+++ b/Engine/Plugins/Experimental/MovieSceneMixedControlRig/MovieSceneMixedControlRig.uplugin
@@ -0,0 +1,44 @@
+{
+	"FileVersion" : 1,
+	"Version" : 1,
+	"VersionName" : "0.1",
+	"FriendlyName" : "Sequencer Mixed Control Rig",
+	"Description" : "System for using the Anim Mixer to mix control rig tracks",
+	"Category" : "Animation",
+	"CreatedBy" : "Epic Games, Inc.",
+	"CreatedByURL" : "https://epicgames.com",
+	"DocsURL" : "",
+	"MarketplaceURL" : "",
+	"SupportURL" : "",
+	"EnabledByDefault" : false,
+	"CanContainContent" : true,
+	"IsBetaVersion": false,
+	"IsExperimentalVersion": true,
+	"Installed" : false,
+	"Plugins": [
+		{
+			"Name": "MovieSceneAnimMixer",
+			"Enabled": true
+		},
+		{
+			"Name": "ControlRig",
+			"Enabled": true
+		},
+		{
+			"Name": "RigVM",
+			"Enabled": true
+		},
+		{
+			"Name": "UAFControlRig",
+			"Enabled": true
+		}
+	],
+	"Modules" :
+	[
+		{
+			"Name" : "MovieSceneMixedControlRig",
+			"Type" : "Runtime",
+			"LoadingPhase" : "Default"
+		}
+	]
+}
diff --git a/Engine/Plugins/Experimental/MovieSceneMixedControlRig/Source/MovieSceneMixedControlRig/MovieSceneMixedControlRig.build.cs b/Engine/Plugins/Experimental/MovieSceneMixedControlRig/Source/MovieSceneMixedControlRig/MovieSceneMixedControlRig.build.cs
new file mode 100644
--- /dev/null
+++ b/Engine/Plugins/Experimental/MovieSceneMixedControlRig/Source/MovieSceneMixedControlRig/MovieSceneMixedControlRig.build.cs
@@ -0,0 +1,29 @@
+// Copyright Epic Games, Inc. All Rights Reserved.
+
+namespace UnrealBuildTool.Rules
+{
+	public class MovieSceneMixedControlRig : ModuleRules
+	{
+		public MovieSceneMixedControlRig(ReadOnlyTargetRules Target) : base(Target)
+		{
+			PrivateDependencyModuleNames.AddRange(
+				new string[]
+				{
+					"Core",
+					"CoreUObject",
+					"Engine",
+					"DeveloperSettings",
+					"Settings",
+					"MovieScene",
+					"MovieSceneTracks",
+					"AnimGraphRuntime",
+					"MovieSceneAnimMixer",
+					"ControlRig",
+					"UAFControlRig",
+					"RigVM",
+					"AnimationCore"
+				}
+			);
+		}
+	}
+}
\ No newline at end of file
diff --git a/Engine/Plugins/Experimental/MovieSceneMixedControlRig/Source/MovieSceneMixedControlRig/Private/MovieSceneMixedControlRigModule.cpp b/Engine/Plugins/Experimental/MovieSceneMixedControlRig/Source/MovieSceneMixedControlRig/Private/MovieSceneMixedControlRigModule.cpp
new file mode 100644
--- /dev/null
+++ b/Engine/Plugins/Experimental/MovieSceneMixedControlRig/Source/MovieSceneMixedControlRig/Private/MovieSceneMixedControlRigModule.cpp
@@ -0,0 +1,27 @@
+// Copyright Epic Games, Inc. All Rights Reserved.
+
+#include "MovieSceneMixedControlRigModule.h"
+#include "Modules/ModuleInterface.h"
+#include "Modules/ModuleManager.h"
+
+DEFINE_LOG_CATEGORY(LogMovieSceneMixedControlRig);
+
+#define LOCTEXT_NAMESPACE "FMixedControlRigModule"
+
+namespace UE::MovieScene
+{
+
+void FMovieSceneMixedControlRigModule::StartupModule()
+{
+}
+
+void FMovieSceneMixedControlRigModule::ShutdownModule()
+{
+
+}
+
+}
+
+IMPLEMENT_MODULE(UE::MovieScene::FMovieSceneMixedControlRigModule, MovieSceneMixedControlRig)
+
+#undef LOCTEXT_NAMESPACE
\ No newline at end of file
diff --git a/Engine/Plugins/Experimental/MovieSceneMixedControlRig/Source/MovieSceneMixedControlRig/Private/Systems/MovieSceneMixedControlRigSystem.cpp b/Engine/Plugins/Experimental/MovieSceneMixedControlRig/Source/MovieSceneMixedControlRig/Private/Systems/MovieSceneMixedControlRigSystem.cpp
new file mode 100644
--- /dev/null
+++ b/Engine/Plugins/Experimental/MovieSceneMixedControlRig/Source/MovieSceneMixedControlRig/Private/Systems/MovieSceneMixedControlRigSystem.cpp
@@ -0,0 +1,332 @@
+// Copyright Epic Games, Inc. All Rights Reserved.
+
+#include "Systems/MovieSceneMixedControlRigSystem.h"
+
+#include "Async/TaskGraphInterfaces.h"
+#include "Decorations/MovieSceneScalingAnchors.h"
+#include "EntitySystem/BuiltInComponentTypes.h"
+#include "AnimMixerComponentTypes.h"
+#include "EntitySystem/MovieSceneEntityIDs.h"
+#include "EntitySystem/MovieSceneEntityMutations.h"
+#include "EntitySystem/MovieSceneEntitySystemLinker.h"
+#include "EntitySystem/MovieSceneEntitySystemRunner.h"
+#include "EntitySystem/MovieSceneEntitySystemTask.h"
+#include "EntitySystem/MovieSceneBoundObjectInstantiator.h"
+#include "EntitySystem/MovieSceneBoundSceneComponentInstantiator.h"
+#include "EntitySystem/Interrogation/MovieSceneInterrogationLinker.h"
+#include "MovieScene.h"
+#include "MovieSceneTracksComponentTypes.h"
+#include "PreAnimatedState/MovieScenePreAnimatedComponentTransformStorage.h"
+#include "Sequencer/MovieSceneControlRigComponentTypes.h"
+#include "Sequencer/MovieSceneControlRigParameterTrack.h"
+#include "Sequencer/MovieSceneControlRigSystem.h"
+#include "IControlRigObjectBinding.h"
+#include "ControlRigTask.h"
+#include "Systems/MovieSceneSkeletalAnimationSystem.h"
+#include "ControlRigObjectBinding.h"
+#include "Rigs/FKControlRig.h"
+#include "MovieSceneAnimationMixerTrack.h"
+
+#include UE_INLINE_GENERATED_CPP_BY_NAME(MovieSceneMixedControlRigSystem)
+
+
+namespace UE::MovieScene
+{
+
+bool GMovieSceneControlRigAnimMixerEnabled = true;
+FAutoConsoleVariableRef CVarMovieSceneControlRigAnimMixerEnabled(
+	TEXT("Sequencer.ControlRigAnimMixer.Enabled"),
+	GMovieSceneControlRigAnimMixerEnabled,
+	TEXT("(Default: true) Controls whether the new Anim Mixer library is used for control rig evaluation.\n"),
+	ECVF_Default);
+
+// Adds required anim mixer components onto control rig track entities
+struct FControlRigMixerMutation : IMovieSceneEntityMutation
+{
+	FBuiltInComponentTypes* BuiltInComponents;
+	FMovieSceneTracksComponentTypes* TrackComponents; 
+	FAnimMixerComponentTypes* AnimMixerComponents;
+	FControlRigComponentTypes* ControlRigComponents;
+	UMovieSceneEntitySystemLinker* Linker = nullptr;
+	FControlRigMixerMutation(UMovieSceneEntitySystemLinker* InLinker)
+	{
+		BuiltInComponents = FBuiltInComponentTypes::Get();
+		TrackComponents = FMovieSceneTracksComponentTypes::Get();
+		AnimMixerComponents = FAnimMixerComponentTypes::Get();
+		ControlRigComponents = FControlRigComponentTypes::Get();
+		Linker = InLinker;
+	}
+
+	virtual void CreateMutation(FEntityManager* EntityManager, FComponentMask* InOutEntityComponentTypes) const
+	{
+		InOutEntityComponentTypes->SetAll(
+		{
+			TrackComponents->Tags.AnimMixerPoseProducer,
+			AnimMixerComponents->Priority,
+			AnimMixerComponents->Target,
+			AnimMixerComponents->Task,
+			AnimMixerComponents->MixerEntry,
+			AnimMixerComponents->Tags.RequiresBlending,
+			BuiltInComponents->BoundObject,
+			BuiltInComponents->BoundObjectKey
+		});
+	}
+
+	virtual void InitializeAllocation(FEntityAllocation* Allocation, const FComponentMask& AllocationType) const
+	{
+		const int32 Num = Allocation->Num();
+		FEntityAllocationWriteContext WriteContext = FEntityAllocationWriteContext::NewAllocation();
+		TComponentReader<FControlRigSourceData> ControlRigSourceData = Allocation->ReadComponents(ControlRigComponents->ControlRigSource);
+		TComponentReader<FRootInstanceHandle> RootInstanceHandles = Allocation->ReadComponents(BuiltInComponents->RootInstanceHandle);
+		TComponentReader<FGuid> ObjectBindingIDs = Allocation->ReadComponents(BuiltInComponents->GenericObjectBinding);
+		TOptionalComponentReader<FMovieSceneSequenceID> OptSequenceIDs = Allocation->TryReadComponents(BuiltInComponents->SequenceID);
+		
+		TComponentWriter<TInstancedStruct<FMovieSceneMixedAnimationTarget>> OutAnimTargets = Allocation->WriteComponents(AnimMixerComponents->Target, WriteContext);
+		TComponentWriter<TSharedPtr<FAnimNextEvaluationTask>> OutTasks = Allocation->WriteComponents(AnimMixerComponents->Task, WriteContext);
+		TComponentWriter<TSharedPtr<FMovieSceneAnimMixerEntry>> OutMixerEntries = Allocation->WriteComponents(AnimMixerComponents->MixerEntry, WriteContext);
+		TComponentWriter<int32> OutPriorities = Allocation->WriteComponents(AnimMixerComponents->Priority, WriteContext);
+		
+		TComponentWriter<UObject*> OutBoundObjects = Allocation->WriteComponents(BuiltInComponents->BoundObject, WriteContext);
+		TComponentWriter<FObjectKey> OutBoundObjectKeys = Allocation->WriteComponents(BuiltInComponents->BoundObjectKey, WriteContext);
+		
+		for (int32 Index = 0; Index < Num; ++Index)
+		{
+			// Set some defaults
+			OutTasks[Index] = nullptr;
+			OutMixerEntries[Index] = nullptr; 
+			OutBoundObjects[Index] = nullptr;
+			OutBoundObjectKeys[Index] = nullptr;
+
+			const FControlRigSourceData& ControlRigSource = ControlRigSourceData[Index];
+			if (ControlRigSource.Track)
+			{
+				// Bit of a hack- we resolve the object binding here because the original Control Rig track was never built to work with resolved object bindings
+				// but the mixer requires them. We only resolve the first object because control rig does not support multi-object bindings.
+				FMovieSceneSequenceID SequenceID = OptSequenceIDs ? OptSequenceIDs[Index] : MovieSceneSequenceID::Root;
+				TArrayView<TWeakObjectPtr<>> BoundObjects = Linker->GetInstanceRegistry()->GetInstance(RootInstanceHandles[Index]).GetSharedPlaybackState()->FindBoundObjects(ObjectBindingIDs[Index], SequenceID);
+				if (!BoundObjects.IsEmpty())
+				{
+					UObject* BoundObject = UMovieSceneSkeletalAnimationSystem::ResolveSkeletalMeshComponentBinding(BoundObjects[0].Get());
+					if (BoundObject)
+					{
+						OutBoundObjects[Index] = BoundObject;
+						OutBoundObjectKeys[Index] = BoundObject;
+
+						UWorld* GameWorld = (BoundObject && BoundObject->GetWorld() && BoundObject->GetWorld()->IsGameWorld()) ? BoundObject->GetWorld() : nullptr;
+						UControlRig* ControlRig = GameWorld ? ControlRigSource.Track->GetGameWorldControlRig(GameWorld) : ControlRigSource.Track->GetControlRig();
+
+						if (ControlRig)
+						{
+							if (!ControlRig->GetObjectBinding())
+							{
+								ControlRig->SetObjectBinding(MakeShared<FControlRigObjectBinding>());
+							}
+
+							if (ControlRig->GetObjectBinding()->GetBoundObject() != FControlRigObjectBinding::GetBindableObject(BoundObject))
+							{
+								ControlRig->GetObjectBinding()->BindToObject(BoundObject);
+							}
+
+							FAnimNextControlRigTaskParams Params;
+							Params.bConsumesPreviousPose = false;
+							Params.ControlRig = ControlRig;
+
+							OutTasks[Index] = MakeShared<FAnimNextControlRigTask>(MoveTemp(Params));
+							if (TSharedPtr<FAnimNextControlRigTask> ControlRigTask = StaticCastSharedPtr<FAnimNextControlRigTask>(OutTasks[Index]))
+							{
+								TArray<FName> SelectedControls = ControlRig->CurrentControlSelection();
+								
+								ControlRigTask->InitializeControlRig();
+
+								if (USkeletalMeshComponent* SkeletalMeshComponent = Cast<USkeletalMeshComponent>(FControlRigObjectBinding::GetBindableObject(BoundObject)))
+								{
+									ControlRig->SetBoneInitialTransformsFromSkeletalMeshComponent(SkeletalMeshComponent, false);
+									ControlRig->Evaluate_AnyThread();
+								}
+
+								if (GameWorld == nullptr && ControlRig->IsA<UFKControlRig>())// mz only in editor replace Fk Control rig, will look post 29.20 to see if this really needed but want to unblock folks
+								{
+									ControlRigSource.Track->ReplaceControlRig(ControlRig, true);
+								}
+
+								TArray<FName> NewSelectedControls = ControlRig->CurrentControlSelection();
+								if (SelectedControls != NewSelectedControls)
+								{
+									ControlRig->ClearControlSelection();
+									for (const FName& Name : SelectedControls)
+									{
+										ControlRig->SelectControl(Name, true);
+									}
+								}
+							}
+						}
+						UMovieSceneAnimationMixerTrack* MixerTrack = Cast<UMovieSceneAnimationMixerTrack>(ControlRigSource.Track->GetOuter());
+
+						// If we have a mixer track, the animation decoration may already have added our target and priority
+
+						if (MixerTrack == nullptr || !MixerTrack->MixedAnimationTarget.IsValid())
+						{
+							// Set a default target- Automatic will resolve to the appropriate place, or the Anim Mixer track will override it.
+							OutAnimTargets[Index] = TInstancedStruct<FMovieSceneMixedAnimationTarget>::Make();
+							// Convert the "Priority Order" on the control rig track to the anim mixer's concept of priority. We do it this way also
+							// to ensure all control rigs (by default anyway) happen after the Animation track.
+							// Control rig tracks start at priority 100 and then decrease from there, with larger numbers evaluating first
+							// Anim mixer uses priority to mean that larger priorities go later (as they will overwrite lower priority poses)
+							// We use an arbitrary high priority to try to ensure that priority values set by control rig users still will evaluate after animation track
+							const int32 HighPriority = 1000000;
+							OutPriorities[Index] = 1 + (HighPriority - ControlRigSource.Track->GetPriorityOrder());
+						}
+					}
+				}
+			}
+		}
+	}
+};
+
+/** ------------------------------------------------------------------------- */
+/** Task for gathering base control rigs and setting up their evaluation tasks. */
+struct FUpdateBaseControlRigs
+{
+
+	const FInstanceRegistry* InstanceRegistry;
+
+	FUpdateBaseControlRigs(UMovieSceneEntitySystemLinker* InLinker)
+		: InstanceRegistry(InLinker->GetInstanceRegistry())
+	{
+	}
+
+	void ForEachEntity(
+		const FInstanceHandle& InstanceHandle,
+		UObject* BoundObject, 
+		const FControlRigSourceData& ControlRigSource,
+		TSharedPtr<FAnimNextEvaluationTask>& EvalTask) const
+	{
+		if (BoundObject)
+		{
+			UWorld* GameWorld = (BoundObject && BoundObject->GetWorld() && BoundObject->GetWorld()->IsGameWorld()) ? BoundObject->GetWorld() : nullptr;
+			UMovieSceneControlRigParameterTrack* Track = ControlRigSource.Track;
+
+			if (Track == nullptr)
+			{
+				return;
+			}
+
+			UControlRig* ControlRig = GameWorld ? Track->GetGameWorldControlRig(GameWorld) : Track->GetControlRig();
+
+			if (ControlRig == nullptr)
+			{
+				return;
+			}
+
+			UObject* ControlRigBoundObject = ControlRig->GetObjectBinding() != nullptr ? ControlRig->GetObjectBinding()->GetBoundObject() : nullptr;
+			if (!ControlRigBoundObject)
+			{
+				return;
+			}
+
+			if (USkeletalMeshComponent* SkeletalMeshComponent = Cast<USkeletalMeshComponent>(ControlRigBoundObject))
+			{			
+				// Get the full context, so we can get both the current and previous evaluation times.
+				const FSequenceInstance& SequenceInstance = InstanceRegistry->GetInstance(InstanceHandle);
+				const FMovieSceneContext& Context = SequenceInstance.GetContext();
+
+				if (TSharedPtr<FAnimNextControlRigTask> ControlRigTask = StaticCastSharedPtr<FAnimNextControlRigTask>(EvalTask))
+				{
+#if UE_ENABLE_DEBUG_DRAWING
+					ControlRigTask->SetComponentTransform(SkeletalMeshComponent->GetComponentTransform());
+#endif
+					if (ControlRigTask->GetControlRig())
+					{
+						ControlRigTask->GetControlRig()->SetDeltaTime(Context.GetFrameRate().AsSeconds(Context.GetDelta()));
+					}
+				}
+			}
+		}
+	}
+};
+
+
+} // namespace UE::MovieScene
+
+UMovieSceneMixedControlRigSystem::UMovieSceneMixedControlRigSystem(const FObjectInitializer& ObjInit)
+	: Super(ObjInit)
+{
+	using namespace UE::MovieScene;
+
+	FMovieSceneTracksComponentTypes* TrackComponents = FMovieSceneTracksComponentTypes::Get();
+	FAnimMixerComponentTypes* AnimMixerComponents = FAnimMixerComponentTypes::Get();
+	FControlRigComponentTypes* ControlRigComponents = FControlRigComponentTypes::Get();
+
+	Phase = ESystemPhase::Instantiation | ESystemPhase::Scheduling;
+	RelevantComponent = ControlRigComponents->BaseControlRigEvalData;
+
+	if (HasAnyFlags(RF_ClassDefaultObject))
+	{
+		FBuiltInComponentTypes* BuiltInComponents = FBuiltInComponentTypes::Get(); 
+
+		DefineComponentConsumer(GetClass(), ControlRigComponents->BaseControlRigEvalData);
+		DefineComponentConsumer(GetClass(), ControlRigComponents->ControlRigSource);
+
+		DefineImplicitPrerequisite(GetClass(), UMovieSceneControlRigParameterEvaluatorSystem::StaticClass());
+
+		DefineComponentProducer(GetClass(), TrackComponents->Tags.AnimMixerPoseProducer);
+
+		DefineComponentProducer(GetClass(), AnimMixerComponents->Task);
+		DefineComponentProducer(GetClass(), AnimMixerComponents->Target);
+		DefineComponentProducer(GetClass(), AnimMixerComponents->Priority);
+
+	}
+}
+
+void UMovieSceneMixedControlRigSystem::OnRun(FSystemTaskPrerequisites& InPrerequisites, FSystemSubsequentTasks& Subsequents)
+{
+	using namespace UE::MovieScene;
+
+	if (!GMovieSceneControlRigAnimMixerEnabled)
+	{
+		return;
+	}
+
+	FBuiltInComponentTypes* BuiltInComponents = FBuiltInComponentTypes::Get();
+	FMovieSceneTracksComponentTypes* TrackComponents = FMovieSceneTracksComponentTypes::Get();
+	FAnimMixerComponentTypes* AnimMixerComponents = FAnimMixerComponentTypes::Get();
+	FControlRigComponentTypes* ControlRigComponents = FControlRigComponentTypes::Get();
+	TSharedRef<FMovieSceneEntitySystemRunner> Runner = Linker->GetRunner();
+
+	if (Runner->GetCurrentPhase() == ESystemPhase::Instantiation)
+	{
+		FEntityComponentFilter Filter;
+		Filter.All({ BuiltInComponents->RootInstanceHandle, ControlRigComponents->BaseControlRigEvalData, BuiltInComponents->Tags.NeedsLink });
+		Filter.None({ BuiltInComponents->Tags.Ignored });
+
+		// Initialize components for control rig mixing
+		FControlRigMixerMutation ControlRigMixerMutation(Linker);
+		Linker->EntityManager.MutateAll(Filter, ControlRigMixerMutation);
+	}
+}
+
+void UMovieSceneMixedControlRigSystem::OnSchedulePersistentTasks(UE::MovieScene::IEntitySystemScheduler* TaskScheduler)
+{
+	using namespace UE::MovieScene;	
+	
+	if (!GMovieSceneControlRigAnimMixerEnabled)
+	{
+		return;
+	}
+
+	FBuiltInComponentTypes* BuiltInComponents = FBuiltInComponentTypes::Get();
+	FMovieSceneTracksComponentTypes* TrackComponents = FMovieSceneTracksComponentTypes::Get();
+	FAnimMixerComponentTypes* AnimMixerComponents = FAnimMixerComponentTypes::Get();
+	FControlRigComponentTypes* ControlRigComponents = FControlRigComponentTypes::Get();
+
+	// Gather all skel anim section entities that have been marked with the AnimMixerPoseProducer tag
+	FTaskID UpdateTask = FEntityTaskBuilder()
+		.Read(BuiltInComponents->InstanceHandle)
+		.Read(BuiltInComponents->BoundObject)
+		.Read(ControlRigComponents->ControlRigSource)
+		.Write(AnimMixerComponents->Task)
+		.FilterAll({ TrackComponents->Tags.AnimMixerPoseProducer, ControlRigComponents->BaseControlRigEvalData })
+		.FilterNone({ BuiltInComponents->Tags.Ignored })
+		.Schedule_PerEntity<FUpdateBaseControlRigs>(&Linker->EntityManager, TaskScheduler, Linker);
+}
+
diff --git a/Engine/Plugins/Experimental/MovieSceneMixedControlRig/Source/MovieSceneMixedControlRig/Public/MovieSceneMixedControlRigModule.h b/Engine/Plugins/Experimental/MovieSceneMixedControlRig/Source/MovieSceneMixedControlRig/Public/MovieSceneMixedControlRigModule.h
new file mode 100644
--- /dev/null
+++ b/Engine/Plugins/Experimental/MovieSceneMixedControlRig/Source/MovieSceneMixedControlRig/Public/MovieSceneMixedControlRigModule.h
@@ -0,0 +1,20 @@
+// Copyright Epic Games, Inc. All Rights Reserved.
+
+#pragma once
+
+#include "CoreMinimal.h"
+#include "Modules/ModuleInterface.h"
+#include "HAL/IConsoleManager.h"
+
+MOVIESCENEMIXEDCONTROLRIG_API DECLARE_LOG_CATEGORY_EXTERN(LogMovieSceneMixedControlRig, Log, All);
+
+namespace UE::MovieScene
+{
+class FMovieSceneMixedControlRigModule : public IModuleInterface
+{
+	// IModuleInterface interface
+	MOVIESCENEMIXEDCONTROLRIG_API virtual void StartupModule() override;
+	MOVIESCENEMIXEDCONTROLRIG_API virtual void ShutdownModule() override;
+};
+
+}
diff --git a/Engine/Plugins/Experimental/MovieSceneMixedControlRig/Source/MovieSceneMixedControlRig/Public/Systems/MovieSceneMixedControlRigSystem.h b/Engine/Plugins/Experimental/MovieSceneMixedControlRig/Source/MovieSceneMixedControlRig/Public/Systems/MovieSceneMixedControlRigSystem.h
new file mode 100644
--- /dev/null
+++ b/Engine/Plugins/Experimental/MovieSceneMixedControlRig/Source/MovieSceneMixedControlRig/Public/Systems/MovieSceneMixedControlRigSystem.h
@@ -0,0 +1,23 @@
+// Copyright Epic Games, Inc. All Rights Reserved.
+
+#pragma once
+
+#include "EntitySystem/MovieSceneEntityInstantiatorSystem.h"
+#include "MovieSceneMixedControlRigSystem.generated.h"
+
+// System to handle creating evaluation tasks from control rig tracks for the anim mixer.
+UCLASS(MinimalAPI)
+class UMovieSceneMixedControlRigSystem
+	: public UMovieSceneEntityInstantiatorSystem
+{
+public:
+
+	GENERATED_BODY()
+
+	UMovieSceneMixedControlRigSystem(const FObjectInitializer& ObjInit);
+
+private:
+
+	virtual void OnRun(FSystemTaskPrerequisites& InPrerequisites, FSystemSubsequentTasks& Subsequents) override;
+	virtual void OnSchedulePersistentTasks(UE::MovieScene::IEntitySystemScheduler* TaskScheduler) override;
+};
\ No newline at end of file
