Sequencer: Possible fix to MacArm64 build.

diff --git a/Engine/Plugins/Experimental/MovieSceneMixedControlRig/MovieSceneMixedControlRig.uplugin b/Engine/Plugins/Experimental/MovieSceneMixedControlRig/MovieSceneMixedControlRig.uplugin
--- a/Engine/Plugins/Experimental/MovieSceneMixedControlRig/MovieSceneMixedControlRig.uplugin
+++ b/Engine/Plugins/Experimental/MovieSceneMixedControlRig/MovieSceneMixedControlRig.uplugin
@@ -28,6 +28,14 @@
 			"Name": "RigVM",
 			"Enabled": true
 		},
+		{
+			"Name": "UAF",
+			"Enabled": true
+		},
+		{
+			"Name": "UAFAnimGraph",
+			"Enabled": true
+		},
 		{
 			"Name": "UAFControlRig",
 			"Enabled": true
diff --git a/Engine/Plugins/Experimental/MovieSceneMixedControlRig/Source/MovieSceneMixedControlRig/MovieSceneMixedControlRig.build.cs b/Engine/Plugins/Experimental/MovieSceneMixedControlRig/Source/MovieSceneMixedControlRig/MovieSceneMixedControlRig.build.cs
--- a/Engine/Plugins/Experimental/MovieSceneMixedControlRig/Source/MovieSceneMixedControlRig/MovieSceneMixedControlRig.build.cs
+++ b/Engine/Plugins/Experimental/MovieSceneMixedControlRig/Source/MovieSceneMixedControlRig/MovieSceneMixedControlRig.build.cs
@@ -21,7 +21,9 @@ namespace UnrealBuildTool.Rules
 					"ControlRig",
 					"UAFControlRig",
 					"RigVM",
-					"AnimationCore"
+					"AnimationCore",
+					"UAF",
+					"UAFAnimGraph"
 				}
 			);
 		}
