Sequencer: Fix ECS mixer control rig version from flickering the entire sequencer tree when dragging sections/making minor changes to the sequence. This was due to constantly re-initializing the control rig on ECS instantiation.

#Synced-CL 48651424

diff --git a/Engine/Plugins/Experimental/MovieSceneMixedControlRig/Source/MovieSceneMixedControlRig/Private/Systems/MovieSceneMixedControlRigSystem.cpp b/Engine/Plugins/Experimental/MovieSceneMixedControlRig/Source/MovieSceneMixedControlRig/Private/Systems/MovieSceneMixedControlRigSystem.cpp
--- a/Engine/Plugins/Experimental/MovieSceneMixedControlRig/Source/MovieSceneMixedControlRig/Private/Systems/MovieSceneMixedControlRigSystem.cpp
+++ b/Engine/Plugins/Experimental/MovieSceneMixedControlRig/Source/MovieSceneMixedControlRig/Private/Systems/MovieSceneMixedControlRigSystem.cpp
@@ -117,6 +117,7 @@ struct FControlRigMixerMutation : IMovieSceneEntityMutation
 
 						if (ControlRig)
 						{
+							bool bNeedsInit = false;
 							if (!ControlRig->GetObjectBinding())
 							{
 								ControlRig->SetObjectBinding(MakeShared<FControlRigObjectBinding>());
@@ -125,6 +126,7 @@ struct FControlRigMixerMutation : IMovieSceneEntityMutation
 							if (ControlRig->GetObjectBinding()->GetBoundObject() != FControlRigObjectBinding::GetBindableObject(BoundObject))
 							{
 								ControlRig->GetObjectBinding()->BindToObject(BoundObject);
+								bNeedsInit = true;
 							}
 
 							FAnimNextControlRigTaskParams Params;
@@ -135,27 +137,30 @@ struct FControlRigMixerMutation : IMovieSceneEntityMutation
 							if (TSharedPtr<FAnimNextControlRigTask> ControlRigTask = StaticCastSharedPtr<FAnimNextControlRigTask>(OutTasks[Index]))
 							{
 								TArray<FName> SelectedControls = ControlRig->CurrentControlSelection();
-								
-								ControlRigTask->InitializeControlRig();
 
-								if (USkeletalMeshComponent* SkeletalMeshComponent = Cast<USkeletalMeshComponent>(FControlRigObjectBinding::GetBindableObject(BoundObject)))
-								{
-									ControlRig->SetBoneInitialTransformsFromSkeletalMeshComponent(SkeletalMeshComponent, false);
-									ControlRig->Evaluate_AnyThread();
-								}
+								ControlRigTask->Initialize(bNeedsInit);
 
-								if (GameWorld == nullptr && ControlRig->IsA<UFKControlRig>())// mz only in editor replace Fk Control rig, will look post 29.20 to see if this really needed but want to unblock folks
+								if (bNeedsInit)
 								{
-									ControlRigSource.Track->ReplaceControlRig(ControlRig, true);
-								}
+									if (USkeletalMeshComponent* SkeletalMeshComponent = Cast<USkeletalMeshComponent>(FControlRigObjectBinding::GetBindableObject(BoundObject)))
+									{
+										ControlRig->SetBoneInitialTransformsFromSkeletalMeshComponent(SkeletalMeshComponent, false);
+										ControlRig->Evaluate_AnyThread();
+									}
 
-								TArray<FName> NewSelectedControls = ControlRig->CurrentControlSelection();
-								if (SelectedControls != NewSelectedControls)
-								{
-									ControlRig->ClearControlSelection();
-									for (const FName& Name : SelectedControls)
+									if (GameWorld == nullptr && ControlRig->IsA<UFKControlRig>())// mz only in editor replace Fk Control rig, will look post 29.20 to see if this really needed but want to unblock folks
+									{
+										ControlRigSource.Track->ReplaceControlRig(ControlRig, true);
+									}
+
+									TArray<FName> NewSelectedControls = ControlRig->CurrentControlSelection();
+									if (SelectedControls != NewSelectedControls)
 									{
-										ControlRig->SelectControl(Name, true);
+										ControlRig->ClearControlSelection();
+										for (const FName& Name : SelectedControls)
+										{
+											ControlRig->SelectControl(Name, true);
+										}
 									}
 								}
 							}
diff --git a/Engine/Plugins/Experimental/UAF/UAFControlRig/Source/UAFControlRig/Private/ControlRigTask.cpp b/Engine/Plugins/Experimental/UAF/UAFControlRig/Source/UAFControlRig/Private/ControlRigTask.cpp
--- a/Engine/Plugins/Experimental/UAF/UAFControlRig/Source/UAFControlRig/Private/ControlRigTask.cpp
+++ b/Engine/Plugins/Experimental/UAF/UAFControlRig/Source/UAFControlRig/Private/ControlRigTask.cpp
@@ -22,7 +22,7 @@ static TAutoConsoleVariable<int32> CVarControlRigDisableExecutionAnimNext(TEXT("
 DEFINE_STAT(STAT_AnimNext_Task_ControlRig);
 
 
-void FAnimNextControlRigTask::InitializeControlRig(const FControlRigVariableMappings::FCustomPropertyMappings* PropertyMappings)
+void FAnimNextControlRigTask::Initialize(bool bInitControlRig, const FControlRigVariableMappings::FCustomPropertyMappings* PropertyMappings)
 {
 	if (Params.ControlRig)
 	{
@@ -43,8 +43,11 @@ void FAnimNextControlRigTask::InitializeControlRig(const FControlRigVariableMapp
 			}
 		}
 
-		Params.ControlRig->Initialize(true);
-		Params.ControlRig->RequestInit();
+		if (bInitControlRig)
+		{
+			Params.ControlRig->Initialize(true);
+			Params.ControlRig->RequestInit();
+		}
 
 		Params.ControlRigHierarchyMappings.ResetRefPoseSetterHash();
 		Params.ControlRigVariableMappings.ResetCurvesInputToControlCache();
diff --git a/Engine/Plugins/Experimental/UAF/UAFControlRig/Source/UAFControlRig/Private/ControlRigTrait.cpp b/Engine/Plugins/Experimental/UAF/UAFControlRig/Source/UAFControlRig/Private/ControlRigTrait.cpp
--- a/Engine/Plugins/Experimental/UAF/UAFControlRig/Source/UAFControlRig/Private/ControlRigTrait.cpp
+++ b/Engine/Plugins/Experimental/UAF/UAFControlRig/Source/UAFControlRig/Private/ControlRigTrait.cpp
@@ -255,7 +255,7 @@ void FControlRigTrait::FInstanceData::InitializeControlRig(const FExecutionConte
 	check(EvaluationTask);
 	if (ControlRig)
 	{
-		EvaluationTask->InitializeControlRig(&PropertyMappings);
+		EvaluationTask->Initialize(true, &PropertyMappings);
 
 		ControlRig->OnInitialized_AnyThread().RemoveAll(this);
 		OnInitializedHandle = ControlRig->OnInitialized_AnyThread().AddRaw(this, &FControlRigTrait::FInstanceData::HandleOnInitialized_AnyThread);
diff --git a/Engine/Plugins/Experimental/UAF/UAFControlRig/Source/UAFControlRig/Public/ControlRigTask.h b/Engine/Plugins/Experimental/UAF/UAFControlRig/Source/UAFControlRig/Public/ControlRigTask.h
--- a/Engine/Plugins/Experimental/UAF/UAFControlRig/Source/UAFControlRig/Public/ControlRigTask.h
+++ b/Engine/Plugins/Experimental/UAF/UAFControlRig/Source/UAFControlRig/Public/ControlRigTask.h
@@ -89,7 +89,7 @@ public:
 	const FAnimNextControlRigTaskParams& GetParams() const { return Params; }
 
 	// Initializes the control rig
-	void InitializeControlRig(const FControlRigVariableMappings::FCustomPropertyMappings* PropertyMappings = nullptr);
+	void Initialize(bool bInitControlRig, const FControlRigVariableMappings::FCustomPropertyMappings* PropertyMappings = nullptr);
 
 	// Marks the control rig as needing initialization on the next execute
 	void RequestInit();
