Sequencer: Control rig ECS Mixer.
Fix multiple sections and additive sections on the control rig ECS when going through the Anim Mixer.

We needed to replicate the same behavior of only creating a single evaluation (evaluation task in this case) for a control rig track with multiple sections, as well as ignoring the weight track for control rig sections in the evaluation task because it's already applied to the controls.

#Synced-CL 50198662

diff --git a/Engine/Plugins/Experimental/MovieSceneAnimMixer/Source/MovieSceneAnimMixer/Private/AnimMixerComponentTypes.cpp b/Engine/Plugins/Experimental/MovieSceneAnimMixer/Source/MovieSceneAnimMixer/Private/AnimMixerComponentTypes.cpp
--- a/Engine/Plugins/Experimental/MovieSceneAnimMixer/Source/MovieSceneAnimMixer/Private/AnimMixerComponentTypes.cpp
+++ b/Engine/Plugins/Experimental/MovieSceneAnimMixer/Source/MovieSceneAnimMixer/Private/AnimMixerComponentTypes.cpp
@@ -54,6 +54,7 @@ namespace UE::MovieScene
 
 		Tags.RequiresBlending = ComponentRegistry->NewTag(TEXT("Requires Blending"));
 		Tags.Transition = ComponentRegistry->NewTag(TEXT("Transition"));
+		Tags.SkipPoseWeight = ComponentRegistry->NewTag(TEXT("Skip Pose Weight"));
 
 		ComponentRegistry->Factories.DuplicateChildComponent(Priority);
 		ComponentRegistry->Factories.DuplicateChildComponent(Target);
diff --git a/Engine/Plugins/Experimental/MovieSceneAnimMixer/Source/MovieSceneAnimMixer/Private/Systems/MovieSceneAnimMixerSystem.cpp b/Engine/Plugins/Experimental/MovieSceneAnimMixer/Source/MovieSceneAnimMixer/Private/Systems/MovieSceneAnimMixerSystem.cpp
--- a/Engine/Plugins/Experimental/MovieSceneAnimMixer/Source/MovieSceneAnimMixer/Private/Systems/MovieSceneAnimMixerSystem.cpp
+++ b/Engine/Plugins/Experimental/MovieSceneAnimMixer/Source/MovieSceneAnimMixer/Private/Systems/MovieSceneAnimMixerSystem.cpp
@@ -66,7 +66,10 @@ namespace UE::MovieScene
 			if (MixerEntry)
 			{
 				MixerEntry->EvalTask = Task;
-				MixerEntry->PoseWeight = WeightAndEasings ? *WeightAndEasings : 1.0;
+
+				// If bSkipPoseWeight is set, use 1.0 for weight
+				MixerEntry->PoseWeight = (MixerEntry->bSkipPoseWeight || !WeightAndEasings) ? 1.0 : *WeightAndEasings;
+
 				if (RootMotionSettings)
 				{
 					MixerEntry->RootMotionSettings = *RootMotionSettings;
@@ -218,6 +221,10 @@ namespace UE::MovieScene
 					{
 						continue;
 					}
+					if (!NextMixerEntry->EvalTask.IsValid())
+					{
+						continue;
+					}
 					// Skip entries whose source sections are handled by an active transition
 					if (SectionsInTransition.Contains(NextMixerEntry->EntityOwner))
 					{
@@ -859,6 +866,7 @@ void UMovieSceneAnimMixerSystem::OnRun(FSystemTaskPrerequisites& InPrerequisites
 				const FComponentMask& Type = Linker->EntityManager.GetEntityType(EntityID);
 				Entry->bAdditive = Type.Contains(BuiltInComponents->Tags.AdditiveAnimation);
 				Entry->bRequiresBlend = Type.Contains(AnimMixerComponents->Tags.RequiresBlending);
+				Entry->bSkipPoseWeight = Type.Contains(AnimMixerComponents->Tags.SkipPoseWeight);
 				Entry->MirrorTable = OptMirrorDataTable ? *OptMirrorDataTable : nullptr;
 				Entry->bIsTransition = Type.Contains(AnimMixerComponents->Tags.Transition);
 
diff --git a/Engine/Plugins/Experimental/MovieSceneAnimMixer/Source/MovieSceneAnimMixer/Public/AnimMixerComponentTypes.h b/Engine/Plugins/Experimental/MovieSceneAnimMixer/Source/MovieSceneAnimMixer/Public/AnimMixerComponentTypes.h
--- a/Engine/Plugins/Experimental/MovieSceneAnimMixer/Source/MovieSceneAnimMixer/Public/AnimMixerComponentTypes.h
+++ b/Engine/Plugins/Experimental/MovieSceneAnimMixer/Source/MovieSceneAnimMixer/Public/AnimMixerComponentTypes.h
@@ -52,6 +52,7 @@ namespace UE::MovieScene
 		{
 			FComponentTypeID RequiresBlending;
 			FComponentTypeID Transition;
+			FComponentTypeID SkipPoseWeight;
 		} Tags;
 
 
diff --git a/Engine/Plugins/Experimental/MovieSceneAnimMixer/Source/MovieSceneAnimMixer/Public/Systems/MovieSceneAnimMixerSystem.h b/Engine/Plugins/Experimental/MovieSceneAnimMixer/Source/MovieSceneAnimMixer/Public/Systems/MovieSceneAnimMixerSystem.h
--- a/Engine/Plugins/Experimental/MovieSceneAnimMixer/Source/MovieSceneAnimMixer/Public/Systems/MovieSceneAnimMixerSystem.h
+++ b/Engine/Plugins/Experimental/MovieSceneAnimMixer/Source/MovieSceneAnimMixer/Public/Systems/MovieSceneAnimMixerSystem.h
@@ -37,6 +37,7 @@ struct FMovieSceneAnimMixerEntry
 	int32 Priority = 0;
 	bool bAdditive = false;
 	bool bRequiresBlend = true;
+	bool bSkipPoseWeight = false;
 	TObjectPtr<UMirrorDataTable> MirrorTable;
 	UE::MovieScene::FMovieSceneEntityID EntityID;
 	UE::MovieScene::FInstanceHandle InstanceHandle;
diff --git a/Engine/Plugins/Experimental/MovieSceneMixedControlRig/Source/MovieSceneMixedControlRig/Private/Systems/MovieSceneMixedControlRigSystem.cpp b/Engine/Plugins/Experimental/MovieSceneMixedControlRig/Source/MovieSceneMixedControlRig/Private/Systems/MovieSceneMixedControlRigSystem.cpp
--- a/Engine/Plugins/Experimental/MovieSceneMixedControlRig/Source/MovieSceneMixedControlRig/Private/Systems/MovieSceneMixedControlRigSystem.cpp
+++ b/Engine/Plugins/Experimental/MovieSceneMixedControlRig/Source/MovieSceneMixedControlRig/Private/Systems/MovieSceneMixedControlRigSystem.cpp
@@ -40,14 +40,35 @@ FAutoConsoleVariableRef CVarMovieSceneControlRigAnimMixerEnabled(
 	TEXT("(Default: true) Controls whether the new Anim Mixer library is used for control rig evaluation.\n"),
 	ECVF_Default);
 
+struct FEncounteredTrackTaskData
+{
+	FMovieSceneEntityID EntityID;
+	bool bHasWeight = false;
+};
+
 // Adds required anim mixer components onto control rig track entities
 struct FControlRigMixerMutation : IMovieSceneEntityMutation
 {
 	FBuiltInComponentTypes* BuiltInComponents;
-	FMovieSceneTracksComponentTypes* TrackComponents; 
+	FMovieSceneTracksComponentTypes* TrackComponents;
 	FAnimMixerComponentTypes* AnimMixerComponents;
 	FControlRigComponentTypes* ControlRigComponents;
 	UMovieSceneEntitySystemLinker* Linker = nullptr;
+
+	// Track encountered tracks to prevent duplicate task creation
+	// Maps Track -> task data for the "active" entity
+	// Using Track (not ControlRig) as key because multiple tracks on different
+	// layers can use the same UControlRig, but each track should have its own task
+	mutable TMap<UMovieSceneControlRigParameterTrack*, FEncounteredTrackTaskData> EncounteredTrackTasks;
+
+	// Entity IDs that need their tasks cleared after the mutation completes
+	// (entities that were demoted due to weight priority)
+	mutable TArray<FMovieSceneEntityID> EntitiesToClearTask;
+
+	// Entity IDs that need the SkipPoseWeight tag added after the mutation completes
+	// (entities with non-absolute blend type where weight is already applied to parameters)
+	mutable TArray<FMovieSceneEntityID> EntitiesToAddSkipPoseWeight;
+
 	FControlRigMixerMutation(UMovieSceneEntitySystemLinker* InLinker)
 	{
 		BuiltInComponents = FBuiltInComponentTypes::Get();
@@ -76,12 +97,14 @@ struct FControlRigMixerMutation : IMovieSceneEntityMutation
 	{
 		const int32 Num = Allocation->Num();
 		FEntityAllocationWriteContext WriteContext = FEntityAllocationWriteContext::NewAllocation();
+		TArrayView<const FMovieSceneEntityID> EntityIDs = Allocation->GetEntityIDs();
 		TComponentReader<FControlRigSourceData> ControlRigSourceData = Allocation->ReadComponents(ControlRigComponents->ControlRigSource);
 		TComponentReader<FRootInstanceHandle> RootInstanceHandles = Allocation->ReadComponents(BuiltInComponents->RootInstanceHandle);
 		TComponentReader<FGuid> ObjectBindingIDs = Allocation->ReadComponents(BuiltInComponents->GenericObjectBinding);
 		TOptionalComponentReader<FMovieSceneSequenceID> OptSequenceIDs = Allocation->TryReadComponents(BuiltInComponents->SequenceID);
 		TOptionalComponentReader<TObjectPtr<UMovieSceneAnimationMixerLayer>> OptMixerLayers = Allocation->TryReadComponents(AnimMixerComponents->MixerLayer);
-		
+		TOptionalComponentReader<double> OptWeightAndEasing = Allocation->TryReadComponents(BuiltInComponents->WeightAndEasingResult);
+
 		TComponentWriter<TInstancedStruct<FMovieSceneMixedAnimationTarget>> OutAnimTargets = Allocation->WriteComponents(AnimMixerComponents->Target, WriteContext);
 		TComponentWriter<TSharedPtr<FAnimNextEvaluationTask>> OutTasks = Allocation->WriteComponents(AnimMixerComponents->Task, WriteContext);
 		TComponentWriter<TSharedPtr<FMovieSceneAnimMixerEntry>> OutMixerEntries = Allocation->WriteComponents(AnimMixerComponents->MixerEntry, WriteContext);
@@ -118,55 +141,123 @@ struct FControlRigMixerMutation : IMovieSceneEntityMutation
 
 						if (ControlRig)
 						{
-							bool bNeedsInit = false;
-							if (!ControlRig->GetObjectBinding())
-							{
-								ControlRig->SetObjectBinding(MakeShared<FControlRigObjectBinding>());
-							}
+							const bool bHasWeight = OptWeightAndEasing != nullptr;
+							UMovieSceneControlRigParameterTrack* Track = ControlRigSource.Track;
 
-							if (ControlRig->GetObjectBinding()->GetBoundObject() != FControlRigObjectBinding::GetBindableObject(BoundObject))
+							// Lambda to initialize control rig binding
+							auto InitializeControlRigBinding = [](UControlRig* InControlRig, UObject* InBoundObject) -> bool
 							{
-								ControlRig->GetObjectBinding()->BindToObject(BoundObject);
-								bNeedsInit = true;
-							}
-
-							FAnimNextControlRigTaskParams Params;
-							Params.bConsumesPreviousPose = false;
-							Params.ControlRig = ControlRig;
+								bool bNeedsInit = false;
+								if (!InControlRig->GetObjectBinding())
+								{
+									InControlRig->SetObjectBinding(MakeShared<FControlRigObjectBinding>());
+								}
 
-							OutTasks[Index] = MakeShared<FAnimNextControlRigTask>(MoveTemp(Params));
-							if (TSharedPtr<FAnimNextControlRigTask> ControlRigTask = StaticCastSharedPtr<FAnimNextControlRigTask>(OutTasks[Index]))
+								if (InControlRig->GetObjectBinding()->GetBoundObject() != FControlRigObjectBinding::GetBindableObject(InBoundObject))
+								{
+									InControlRig->GetObjectBinding()->BindToObject(InBoundObject);
+									bNeedsInit = true;
+								}
+								return bNeedsInit;
+							};
+
+							// Lambda to initialize control rig task
+							auto InitializeControlRigTask = [&ControlRigSource, GameWorld](
+								TSharedPtr<FAnimNextEvaluationTask>& Task,
+								UControlRig* InControlRig,
+								UObject* InBoundObject,
+								bool bNeedsInit)
 							{
-								TArray<FName> SelectedControls = ControlRig->CurrentControlSelection();
+								if (TSharedPtr<FAnimNextControlRigTask> ControlRigTask = StaticCastSharedPtr<FAnimNextControlRigTask>(Task))
+								{
+									TArray<FName> SelectedControls = InControlRig->CurrentControlSelection();
 
-								ControlRigTask->Initialize(bNeedsInit);
+									ControlRigTask->Initialize(bNeedsInit);
 
-								if (bNeedsInit)
-								{
-									if (USkeletalMeshComponent* SkeletalMeshComponent = Cast<USkeletalMeshComponent>(FControlRigObjectBinding::GetBindableObject(BoundObject)))
+									if (bNeedsInit)
 									{
-										ControlRig->SetBoneInitialTransformsFromSkeletalMeshComponent(SkeletalMeshComponent, false);
-										ControlRig->Evaluate_AnyThread();
-									}
+										if (USkeletalMeshComponent* SkeletalMeshComponent = Cast<USkeletalMeshComponent>(FControlRigObjectBinding::GetBindableObject(InBoundObject)))
+										{
+											InControlRig->SetBoneInitialTransformsFromSkeletalMeshComponent(SkeletalMeshComponent, false);
+											InControlRig->Evaluate_AnyThread();
+										}
 
-									if (GameWorld == nullptr && ControlRig->IsA<UFKControlRig>())// mz only in editor replace Fk Control rig, will look post 29.20 to see if this really needed but want to unblock folks
-									{
-										ControlRigSource.Track->ReplaceControlRig(ControlRig, true);
-									}
+										if (GameWorld == nullptr && InControlRig->IsA<UFKControlRig>())
+										{
+											ControlRigSource.Track->ReplaceControlRig(InControlRig, true);
+										}
 
-									TArray<FName> NewSelectedControls = ControlRig->CurrentControlSelection();
-									if (SelectedControls != NewSelectedControls)
-									{
-										ControlRig->ClearControlSelection();
-										for (const FName& Name : SelectedControls)
+										TArray<FName> NewSelectedControls = InControlRig->CurrentControlSelection();
+										if (SelectedControls != NewSelectedControls)
 										{
-											ControlRig->SelectControl(Name, true);
+											InControlRig->ClearControlSelection();
+											for (const FName& Name : SelectedControls)
+											{
+												InControlRig->SelectControl(Name, true);
+											}
 										}
 									}
 								}
+							};
+
+							// Check if we've already created a task for this Track
+							const FMovieSceneEntityID EntityID = EntityIDs[Index];
+							FEncounteredTrackTaskData* ExistingData = EncounteredTrackTasks.Find(Track);
+							if (ExistingData)
+							{
+								if (ExistingData->bHasWeight)
+								{
+									// Previous entity has weight - it keeps priority
+									// Current entity doesn't get a task
+									OutTasks[Index] = nullptr;
+								}
+								else if (bHasWeight)
+								{
+									// Previous has NO weight, current DOES - demote previous, promote current
+									// Add previous entity to the list to clear after mutation completes
+									EntitiesToClearTask.Add(ExistingData->EntityID);
+
+									// Create new task for current entity
+									bool bNeedsInit = InitializeControlRigBinding(ControlRig, BoundObject);
+
+									FAnimNextControlRigTaskParams Params;
+									Params.bConsumesPreviousPose = false;
+									Params.ControlRig = ControlRig;
+
+									OutTasks[Index] = MakeShared<FAnimNextControlRigTask>(MoveTemp(Params));
+									InitializeControlRigTask(OutTasks[Index], ControlRig, BoundObject, bNeedsInit);
+
+									// Update tracking
+									ExistingData->EntityID = EntityID;
+									ExistingData->bHasWeight = true;
+								}
+								else
+								{
+									// Neither has weight - previous keeps priority
+									OutTasks[Index] = nullptr;
+								}
+							}
+							else
+							{
+								// First time seeing this track - create new task
+								bool bNeedsInit = InitializeControlRigBinding(ControlRig, BoundObject);
+
+								FAnimNextControlRigTaskParams Params;
+								Params.bConsumesPreviousPose = false;
+								Params.ControlRig = ControlRig;
+
+								OutTasks[Index] = MakeShared<FAnimNextControlRigTask>(MoveTemp(Params));
+								InitializeControlRigTask(OutTasks[Index], ControlRig, BoundObject, bNeedsInit);
+
+								// Track this track
+								FEncounteredTrackTaskData NewData;
+								NewData.EntityID = EntityID;
+								NewData.bHasWeight = bHasWeight;
+								EncounteredTrackTasks.Add(Track, NewData);
 							}
 						}
 
+						// Always set up mixer components (even for entities without tasks)
 						// If we haven't had a mixer layer set, then set default anim target/priority
 						if (!OptMixerLayers || !OptMixerLayers[Index])
 						{
@@ -180,6 +271,16 @@ struct FControlRigMixerMutation : IMovieSceneEntityMutation
 							const int32 HighPriority = 1000000;
 							OutPriorities[Index] = 1 + (HighPriority - ControlRigSource.Track->GetPriorityOrder());
 						}
+
+						// Track entities that need SkipPoseWeight tag (non-absolute blend types).
+						// For non-absolute blends (Additive, Relative, Override), weight is already applied
+						// to control parameters during ECS blending, so we skip pose weight to avoid double-weighting.
+						const bool bIsAbsoluteBlend = ControlRigSource.Section.IsValid() && ControlRigSource.Section->GetBlendType() == EMovieSceneBlendType::Absolute;
+						if (!bIsAbsoluteBlend)
+						{
+							const FMovieSceneEntityID EntityID = EntityIDs[Index];
+							EntitiesToAddSkipPoseWeight.Add(EntityID);
+						}
 					}
 				}
 			}
@@ -306,6 +407,32 @@ void UMovieSceneMixedControlRigSystem::OnRun(FSystemTaskPrerequisites& InPrerequ
 		// Initialize components for control rig mixing
 		FControlRigMixerMutation ControlRigMixerMutation(Linker);
 		Linker->EntityManager.MutateAll(Filter, ControlRigMixerMutation);
+
+		// Clear tasks for entities that were demoted due to weight priority
+		if (ControlRigMixerMutation.EntitiesToClearTask.Num() > 0)
+		{
+			TSet<FMovieSceneEntityID> EntitiesToClear(ControlRigMixerMutation.EntitiesToClearTask);
+
+			FEntityTaskBuilder()
+				.ReadEntityIDs()
+				.Write(AnimMixerComponents->Task)
+				.FilterAll({ TrackComponents->Tags.AnimMixerPoseProducer })
+				.Iterate_PerEntity(&Linker->EntityManager,
+					[&EntitiesToClear](FMovieSceneEntityID EntityID, TSharedPtr<FAnimNextEvaluationTask>& Task)
+					{
+						if (EntitiesToClear.Contains(EntityID))
+						{
+							Task = nullptr;
+						}
+					});
+		}
+
+		// Add SkipPoseWeight tag to entities with non-absolute blend type
+		// (weight is already applied to their parameters during ECS blending)
+		for (const FMovieSceneEntityID& EntityID : ControlRigMixerMutation.EntitiesToAddSkipPoseWeight)
+		{
+			Linker->EntityManager.AddComponent(EntityID, AnimMixerComponents->Tags.SkipPoseWeight);
+		}
 	}
 }
 
